Cinematic template initial checkin

#rnx
#virtualized
#Synced-CL 51828360

diff --git a/Templates/TP_ME_CineBP/Config/DefaultEngine.ini b/Templates/TP_ME_CineBP/Config/DefaultEngine.ini
new file mode 100644
--- /dev/null
+++ b/Templates/TP_ME_CineBP/Config/DefaultEngine.ini
@@ -0,0 +1,56 @@
+[/Script/HardwareTargeting.HardwareTargetingSettings]
+TargetedHardwareClass=Desktop
+AppliedTargetedHardwareClass=Desktop
+DefaultGraphicsPerformance=Maximum
+AppliedDefaultGraphicsPerformance=Maximum
+
+[/Script/Engine.RendererSettings]
+; Lumen GI + Reflections
+r.DynamicGlobalIlluminationMethod=1
+r.ReflectionMethod=1
+r.GenerateMeshDistanceFields=True
+r.Shadow.Virtual.Enable=1
+r.VirtualTextures=True
+
+; Hardware Ray Tracing
+r.RayTracing=True
+r.RayTracing.RayTracingProxies.ProjectEnabled=True
+r.RayTracing.UseTextureLod=True
+r.PathTracing=True
+
+; Lumen Quality
+r.Lumen.HardwareRayTracing=True
+r.Lumen.HardwareRayTracing.LightingMode=1
+r.Lumen.TranslucencyReflections.FrontLayer.EnableForProject=True
+r.Lumen.ScreenTracingSource=1
+r.Lumen.Reflections.HardwareRayTracing.Translucent.Refraction.EnableForProject=True
+
+; Character Rendering
+r.SkinCache.CompileShaders=True
+
+; Heterogeneous Volumes
+r.HeterogeneousVolumes.Shadows=True
+r.Translucency.HeterogeneousVolumes=True
+
+; Post-processing & Compositing
+r.DefaultFeature.AutoExposure=False
+r.DefaultFeature.AutoExposure.ExtendDefaultLuminanceRange=True
+r.DefaultFeature.MotionBlur=False
+r.DefaultFeature.Bloom=False
+r.DefaultFeature.LensFlare=True
+r.PostProcessing.PropagateAlpha=True
+r.Deferred.SupportPrimitiveAlphaHoldout=True
+
+[/Script/EngineSettings.GameMapsSettings]
+EditorStartupMap=/Game/CinematicTemplate/Maps/Main.Main
+
+[/Script/WindowsTargetPlatform.WindowsTargetSettings]
+DefaultGraphicsRHI=DefaultGraphicsRHI_DX12
+-D3D12TargetedShaderFormats=PCD3D_SM5
++D3D12TargetedShaderFormats=PCD3D_SM6
+-D3D11TargetedShaderFormats=PCD3D_SM5
++D3D11TargetedShaderFormats=PCD3D_SM5
+
+[/Script/LinuxTargetPlatform.LinuxTargetSettings]
+-TargetedRHIs=SF_VULKAN_SM5
++TargetedRHIs=SF_VULKAN_SM6
diff --git a/Templates/TP_ME_CineBP/Config/TemplateDefs.ini b/Templates/TP_ME_CineBP/Config/TemplateDefs.ini
new file mode 100644
--- /dev/null
+++ b/Templates/TP_ME_CineBP/Config/TemplateDefs.ini
@@ -0,0 +1,76 @@
+[/Script/GameProjectGeneration.TemplateProjectDefs]
+
+Categories=ME
+
+LocalizedDisplayNames=(Language="en", Text="Cinematic")
+LocalizedDescriptions=(Language="en", Text="A template optimized for cinematic and linear content creation. Pre-configured with high-quality rendering (Lumen, hardware ray tracing, virtual shadow maps, path tracing) and cinematic plugins (Movie Render Pipeline, Cinematic Assembly Tools, Animator Kit).\n\nIdeal for film, episodic, trailers, cinematics, pre-vis, and animation projects.")
+
+LocalizedDisplayNames=(Language="ko", Text="시네마틱")
+LocalizedDescriptions=(Language="ko", Text="시네마틱 및 리니어 콘텐츠 제작에 최적화된 템플릿입니다. 고품질 렌더링(Lumen, 하드웨어 레이 트레이싱, 버추얼 쉐도우 맵, 패스 트레이싱)과 시네마틱 플러그인(Movie Render Pipeline, Cinematic Assembly Tools, Animator Kit)이 사전 구성되어 있습니다.\n\n영화, 에피소드, 트레일러, 시네마틱, 프리비즈 및 애니메이션 프로젝트에 적합합니다.")
+
+LocalizedDisplayNames=(Language="zh-Hans", Text="过场动画")
+LocalizedDescriptions=(Language="zh-Hans", Text="针对过场动画和线性内容创作优化的模板。预配置了高质量渲染（Lumen、硬件光线追踪、虚拟阴影贴图、路径追踪）和过场动画插件（Movie Render Pipeline、Cinematic Assembly Tools、Animator Kit）。\n\n适用于电影、剧集、预告片、过场动画、预视和动画项目。")
+
+LocalizedDisplayNames=(Language="ja", Text="シネマティック")
+LocalizedDescriptions=(Language="ja", Text="シネマティックおよびリニアコンテンツ制作に最適化されたテンプレートです。高品質レンダリング（Lumen、ハードウェア レイトレーシング、バーチャル シャドウ マップ、パス トレーシング）とシネマティック プラグイン（Movie Render Pipeline、Cinematic Assembly Tools、Animator Kit）が事前設定されています。\n\n映画、エピソード、トレーラー、シネマティック、プリビズ、アニメーション プロジェクトに最適です。")
+
+; The sort key to use for the template to bypass the default alphabetical sorting.
+SortKey="_3"
+
+; This project uses a thumbnail as the project icon
+bThumbnailAsIcon=true
+
+; A list of folders to avoid copying when creating a project from this template. 
+; These should contain files only of concern to the template itself and not desired in the user's project.
+FoldersToIgnore=Binaries
+FoldersToIgnore=Build
+FoldersToIgnore=Intermediate
+FoldersToIgnore=Saved
+FoldersToIgnore=Media
+
+; The files to ignore copying, relative to the template root folder. 
+; %TEMPLATENAME% will be replaced by the template's folder name when resolving these. 
+; Filename comparisons are case-insensitive. 
+; For example, if this template is "TP_Blank" and the user creates "MyProject", then "TP_Blank.uproject" will be copied to "MyProject.uproject".
+FilesToIgnore="%TEMPLATENAME%.uproject"
+FilesToIgnore="%TEMPLATENAME%.png"
+FilesToIgnore="Config/TemplateDefs.ini"
+FilesToIgnore="Manifest.json"
+FilesToIgnore="contents.txt"
+
+; Folders that will be renamed when creating a project from this template.
+; %PROJECTNAME% will be replaced by the project name specified by the user during creation when resolving these rules.
+; For example,  then "Source/TP_Blank" will be copied to "Source/MyProject".
+FolderRenames=(From="Source/%TEMPLATENAME%", To="Source/%PROJECTNAME%")
+FolderRenames=(From="Source/%TEMPLATENAME%Editor", To="Source/%PROJECTNAME%Editor")
+
+; Here are some default rules for renaming files:
+
+; This rule states that a file called "TP_BLANK.cpp" will be renamed to "MYPROJECT.cpp".
+FilenameReplacements=(Extensions=("cpp","h","ini","cs"), From="%TEMPLATENAME_UPPERCASE%", To="%PROJECTNAME_UPPERCASE%", bCaseSensitive=true)
+
+; This rule states that a file called "tp_blank.ini" will be renamed to "myproject.ini"
+FilenameReplacements=(Extensions=("cpp","h","ini","cs"), From="%TEMPLATENAME_LOWERCASE%", To="%PROJECTNAME_LOWERCASE%", bCaseSensitive=true)
+
+; And this rule catches any irregularly-cased files, so that "Tp_bLANK.ini" will be renamed to "MyProject.ini".
+FilenameReplacements=(Extensions=("cpp","h","ini","cs"), From="%TEMPLATENAME%",To="%PROJECTNAME%",bCaseSensitive=false)
+
+; Rules for renaming strings in the contents of files. 
+; The same rules as above apply, but any usages of strings that match these rules inside the file will be changed. This includes all substrings!
+; For example, "TP_Blank" will be replaced by "MyProject", but "TP_Blank_Name" will also be replaced by "MyProject_Name".
+ReplacementsInFiles=(Extensions=("cpp","h","ini","cs"), From="%TEMPLATENAME_UPPERCASE%", To="%PROJECTNAME_UPPERCASE%", bCaseSensitive=true)
+ReplacementsInFiles=(Extensions=("cpp","h","ini","cs"), From="%TEMPLATENAME_LOWERCASE%", To="%PROJECTNAME_LOWERCASE%", bCaseSensitive=true)
+ReplacementsInFiles=(Extensions=("cpp","h","ini","cs"), From="%TEMPLATENAME%", To="%PROJECTNAME%", bCaseSensitive=false)
+
+; Which settings on the project settings page should be hidden for this template. Default: none
+; Allowed values are:
+; All - Hides the settings page entirely
+; Languages - Choose between C++ and Blueprint versions.
+; HardwareTarget - Choose between Mobile and Desktop hardware targets.
+; GraphicsPreset - Choose between Maximum and Scalable graphics.
+; StarterContent - Choose whether to add starter content to the project.
+; XR - Choose whether to enable Virtual Reality support for the project.
+HiddenSettings=Languages
+HiddenSettings=HardwareTarget
+HiddenSettings=GraphicsPreset
+HiddenSettings=XR
diff --git a/Templates/TP_ME_CineBP/TP_ME_CineBP.uproject b/Templates/TP_ME_CineBP/TP_ME_CineBP.uproject
new file mode 100644
--- /dev/null
+++ b/Templates/TP_ME_CineBP/TP_ME_CineBP.uproject
@@ -0,0 +1,32 @@
+{
+	"FileVersion": 3,
+	"EngineAssociation": "",
+	"Category": "",
+	"Description": "",
+	"Plugins": [
+		{
+			"Name": "AnimatorKit",
+			"Enabled": true
+		},
+		{
+			"Name": "CinematicAssemblyTools",
+			"Enabled": true
+		},
+		{
+			"Name": "ConsoleVariables",
+			"Enabled": true
+		},
+		{
+			"Name": "MoviePipelineMaskRenderPass",
+			"Enabled": true
+		},
+		{
+			"Name": "MovieRenderPipeline",
+			"Enabled": true
+		},
+		{
+			"Name": "OpenColorIO",
+			"Enabled": true
+		}
+	]
+}
