Added example cine assembly shot structure

#rnx
#virtualized
#Synced-CL 52390596

diff --git a/Templates/TP_ME_CineBP/Config/DefaultEngine.ini b/Templates/TP_ME_CineBP/Config/DefaultEngine.ini
--- a/Templates/TP_ME_CineBP/Config/DefaultEngine.ini
+++ b/Templates/TP_ME_CineBP/Config/DefaultEngine.ini
@@ -63,3 +63,7 @@ bPreferProtronInGame=True
 
 [/Script/Engine.VirtualTexturePoolConfig]
 DefaultSizeInMegabyte=256
+
+[/Script/CineAssemblyToolsEditor.ProductionSettings]
++Productions=(ProductionID=65321665493CD7B63CE9649DF569D283,ProductionName="CinematicTemplate",DefaultDisplayRate=(Numerator=24,Denominator=1),DefaultStartFrame=100,SubsequencePriority=BottomUp,NamingTokenNamespaceDenyList=(),DefaultAssetNames=(("/Script/CoreUObject.Class'/Script/Engine.Material'", "M_"),("/Script/CoreUObject.Class'/Script/Engine.MaterialInstanceConstant'", "MI_{source}")),TemplateFolders=((InternalPath="/Game/{cat:activeProduction}"),(InternalPath="/Game/{cat:activeProduction}/Cinematic"),(InternalPath="/Game/{cat:activeProduction}/Materials"),(InternalPath="/Game/{cat:activeProduction}/Geometry")),ExtendedDataConfig=())
+bEnableProductionsLevelEditorToolbar=True
\ No newline at end of file
diff --git a/Templates/TP_ME_CineBP/Config/TemplateDefs.ini b/Templates/TP_ME_CineBP/Config/TemplateDefs.ini
--- a/Templates/TP_ME_CineBP/Config/TemplateDefs.ini
+++ b/Templates/TP_ME_CineBP/Config/TemplateDefs.ini
@@ -62,6 +62,9 @@ ReplacementsInFiles=(Extensions=("cpp","h","ini","cs"), From="%TEMPLATENAME_UPPE
 ReplacementsInFiles=(Extensions=("cpp","h","ini","cs"), From="%TEMPLATENAME_LOWERCASE%", To="%PROJECTNAME_LOWERCASE%", bCaseSensitive=true)
 ReplacementsInFiles=(Extensions=("cpp","h","ini","cs"), From="%TEMPLATENAME%", To="%PROJECTNAME%", bCaseSensitive=false)
 
+SharedContentPacks=(MountName="Characters",DetailLevels=("High"))
+EditDetailLevelPreference="High"
+
 ; Which settings on the project settings page should be hidden for this template. Default: none
 ; Allowed values are:
 ; All - Hides the settings page entirely
